Fix a case where UMobilePendingContent can be GCed while manifest download is in progress

diff --git a/Engine/Plugins/Runtime/MobilePatchingUtils/Source/MobilePatchingUtils/Private/MobilePatchingLibrary.cpp b/Engine/Plugins/Runtime/MobilePatchingUtils/Source/MobilePatchingUtils/Private/MobilePatchingLibrary.cpp
--- a/Engine/Plugins/Runtime/MobilePatchingUtils/Source/MobilePatchingUtils/Private/MobilePatchingLibrary.cpp
+++ b/Engine/Plugins/Runtime/MobilePatchingUtils/Source/MobilePatchingUtils/Private/MobilePatchingLibrary.cpp
@@ -333,6 +333,13 @@ static void OnRemoteManifestReady(FHttpRequestPtr Request, FHttpResponsePtr Resp
 			OnFailed.ExecuteIfBound(ErrorText, (int32)ErrorCode);
 		}
 	};
+
+	if (MobilePendingContent)
+	{
+		// We don't need to keep this object alive anymore, 
+		// on Success it will be passed to the User and it's up to him to keep a reference to it
+		MobilePendingContent->RemoveFromRoot();
+	}
 		
 	if (!bSucceeded || !Response.IsValid())
 	{
@@ -432,6 +439,8 @@ void UMobilePatchingLibrary::RequestContent(const FString& RemoteManifestURL, co
 	MobilePendingContent->RemoteManifestURL = RemoteManifestURL;
 	MobilePendingContent->CloudURL = CloudURL;
 	MobilePendingContent->InstalledManifest = GetInstalledManifest(InstallDirectory);
+	// We need to keep this object alive until we pass it to the User
+	MobilePendingContent->AddToRoot();
 
 	// Request remote manifest
 	TSharedRef<class IHttpRequest> HttpRequest = FHttpModule::Get().CreateRequest();
